Fix 'Ambiguous search' warnings appearing that were unrelated to actual warnings or errors during cook

diff --git a/Engine/Source/Editor/UATHelper/UATHelperModule.cpp b/Engine/Source/Editor/UATHelper/UATHelperModule.cpp
--- a/Engine/Source/Editor/UATHelper/UATHelperModule.cpp
+++ b/Engine/Source/Editor/UATHelper/UATHelperModule.cpp
@@ -182,12 +182,17 @@ private:
 				// Convert from the asset's full path provided by UE_ASSET_LOG back to an AssetData, if possible
 				FString LongPackageName;
 				FPaths::NormalizeFilename(AssetPath);
-				if (FPackageName::TryConvertFilenameToLongPackageName(AssetPath, LongPackageName))
+				if (!FPaths::IsRelative(AssetPath) && FPackageName::TryConvertFilenameToLongPackageName(AssetPath, LongPackageName))
 				{
 					// Generate qualified asset path and query the registry
 					AssetPath = LongPackageName + TEXT(".") + FPackageName::GetShortName(LongPackageName);
-					FAssetRegistryModule& AssetRegistryModule = FModuleManager::LoadModuleChecked<FAssetRegistryModule>(TEXT("AssetRegistry"));
-					AssetData = AssetRegistryModule.Get().GetAssetByObjectPath(FName(*AssetPath));
+					FName AssetPathName(*AssetPath, FNAME_Find);
+					if (!AssetPathName.IsNone())
+					{
+						static const FName AssetRegistryModuleName(TEXT("AssetRegistry"));
+						FAssetRegistryModule& AssetRegistryModule = FModuleManager::LoadModuleChecked<FAssetRegistryModule>(AssetRegistryModuleName);
+						AssetData = AssetRegistryModule.Get().GetAssetByObjectPath(AssetPathName, true);
+					}
 				}
 			}
 	
