Updated UATHelper to cache the packaging notification and ensure it's destroyed when a new packaging operation starts

#rn Pre-existing "Packaging Failed" notifications will now automatically be dismissed when packaging

diff --git a/Engine/Source/Editor/UATHelper/UATHelperModule.cpp b/Engine/Source/Editor/UATHelper/UATHelperModule.cpp
--- a/Engine/Source/Editor/UATHelper/UATHelperModule.cpp
+++ b/Engine/Source/Editor/UATHelper/UATHelperModule.cpp
@@ -449,6 +449,10 @@ public:
 		);
 
 		TSharedPtr<SNotificationItem> NotificationItem = FSlateNotificationManager::Get().AddNotification(Info);
+		if(NotificationItemPtr.IsValid())
+		{
+			NotificationItemPtr.Pin().Get()->Fadeout();
+		}
 
 		if (!NotificationItem.IsValid())
 		{
@@ -461,7 +465,7 @@ public:
 		NotificationItem->SetCompletionState(SNotificationItem::CS_Pending);
 
 		// launch the packager
-		TWeakPtr<SNotificationItem> NotificationItemPtr(NotificationItem);
+		NotificationItemPtr = NotificationItem;
 
 		EventData Data;
 		Data.StartTime = FPlatformTime::Seconds();
@@ -647,6 +651,10 @@ public:
 
 		}
 	}
+
+private:
+	TWeakPtr<SNotificationItem> NotificationItemPtr;
+
 };
 
 IMPLEMENT_MODULE(FUATHelperModule, UATHelper)
