Added an option to allow the editor to try to get the users attention when ever a UAT task completes as these tend to take a long time during which the user will often change focus to work in a different program.

The option can be found under "Editor Preferences - General - Miscellaneous - Get Attention On UAT Completion".

Note that the DrawAttention behaviour is currently only implemented for windows platforms, which will result in the editors entry in the task bar to flash.

diff --git a/Engine/Config/BaseEditorPerProjectUserSettings.ini b/Engine/Config/BaseEditorPerProjectUserSettings.ini
--- a/Engine/Config/BaseEditorPerProjectUserSettings.ini
+++ b/Engine/Config/BaseEditorPerProjectUserSettings.ini
@@ -50,6 +50,10 @@ bShowImportDialogAtReimport=False
 
 ; Export
 bKeepAttachHierarchy=true
+
+; UAT
+bGetAttentionOnUATCompletion=True
+
 ; Misc
 HoverHighlightIntensity=0.000000
 bDisplayUIExtensionPoints=False
diff --git a/Engine/Source/Editor/UATHelper/UATHelperModule.cpp b/Engine/Source/Editor/UATHelper/UATHelperModule.cpp
--- a/Engine/Source/Editor/UATHelper/UATHelperModule.cpp
+++ b/Engine/Source/Editor/UATHelper/UATHelperModule.cpp
@@ -21,6 +21,9 @@
 #include "IUATHelperModule.h"
 #include "AssetRegistryModule.h"
 
+#include "Editor/MainFrame/Public/Interfaces/IMainFrameModule.h"
+#include "Editor/EditorPerProjectUserSettings.h"
+
 #include "Framework/Notifications/NotificationManager.h"
 #include "Widgets/Notifications/SNotificationList.h"
 #include "Logging/TokenizedMessage.h"
@@ -126,6 +129,21 @@ public:
 				ExpireNotificationItemPtr = NotificationItem;
 				NotificationItem->SetCompletionState(CompletionState);
 			}
+
+			// Since the task was probably fairly long, we should try and grab the users attention if they have the option enabled.
+			const UEditorPerProjectUserSettings* SettingsPtr = GetDefault<UEditorPerProjectUserSettings>();
+			if (SettingsPtr->bGetAttentionOnUATCompletion)
+			{
+				IMainFrameModule* MainFrame = FModuleManager::LoadModulePtr<IMainFrameModule>("MainFrame");
+				if (MainFrame != nullptr)
+				{
+					TSharedPtr<SWindow> ParentWindow = MainFrame->GetParentWindow();
+					if (ParentWindow != nullptr)
+					{
+						ParentWindow->DrawAttention(FWindowDrawAttentionParameters(EWindowDrawAttentionRequestType::UntilActivated));
+					}
+				}
+			}
 		}
 	}
 
diff --git a/Engine/Source/Editor/UnrealEd/Classes/Editor/EditorPerProjectUserSettings.h b/Engine/Source/Editor/UnrealEd/Classes/Editor/EditorPerProjectUserSettings.h
--- a/Engine/Source/Editor/UnrealEd/Classes/Editor/EditorPerProjectUserSettings.h
+++ b/Engine/Source/Editor/UnrealEd/Classes/Editor/EditorPerProjectUserSettings.h
@@ -109,6 +109,10 @@ class UEditorPerProjectUserSettings : public UObject
 	UPROPERTY(config)
 	bool bSCSEditorShowFloor;
 
+	/** If enabled, the Editor will attempt to get the users attention whenever a UAT task (such as cooking or packaging) is completed */
+	UPROPERTY(EditAnywhere, config, Category = UnrealAutomationTool)
+	bool bGetAttentionOnUATCompletion;
+
 	/** How fast the SCS viewport camera moves */
 	UPROPERTY(config, meta=(UIMin = "1", UIMax = "8", ClampMin="1", ClampMax="8"))
 	int32 SCSViewportCameraSpeed;
