Main frame notifications couldn't be dismissed if they would also time out because the pointer would reset.

diff --git a/Engine/Source/Editor/UATHelper/UATHelperModule.cpp b/Engine/Source/Editor/UATHelper/UATHelperModule.cpp
--- a/Engine/Source/Editor/UATHelper/UATHelperModule.cpp
+++ b/Engine/Source/Editor/UATHelper/UATHelperModule.cpp
@@ -114,10 +114,10 @@ public:
 				NotificationItem->SetHyperlink(FSimpleDelegate::CreateStatic(&HandleHyperlinkNavigate), Message);
 
 			}
-
+			
+			ExpireNotificationItemPtr = NotificationItem;
 			if (bExpireAndFadeout)
 			{
-				ExpireNotificationItemPtr.Reset();
 				NotificationItem->SetExpireDuration(6.0f);
 				NotificationItem->SetFadeOutDuration(0.5f);
 				NotificationItem->SetCompletionState(CompletionState);
@@ -126,7 +126,6 @@ public:
 			else
 			{
 				// Handling the notification expiration in callback
-				ExpireNotificationItemPtr = NotificationItem;
 				NotificationItem->SetCompletionState(CompletionState);
 			}
 
