Fix for Packaging Toast sounds not playing

diff --git a/Engine/Source/Editor/UATHelper/UATHelperModule.cpp b/Engine/Source/Editor/UATHelper/UATHelperModule.cpp
--- a/Engine/Source/Editor/UATHelper/UATHelperModule.cpp
+++ b/Engine/Source/Editor/UATHelper/UATHelperModule.cpp
@@ -91,19 +91,15 @@ public:
 	{
 		if ( NotificationItemPtr.IsValid() )
 		{
-			const ULevelEditorPlaySettings* EditorPlaySettings = GetDefault<ULevelEditorPlaySettings>();
-			if (EditorPlaySettings->EnablePIEEnterAndExitSounds)
+			if (CompletionState == SNotificationItem::CS_Fail)
 			{
-				if ( CompletionState == SNotificationItem::CS_Fail )
-				{
-					GEditor->PlayEditorSound(TEXT("/Engine/EditorSounds/Notifications/CompileFailed_Cue.CompileFailed_Cue"));
-				}
-				else
-				{
-					GEditor->PlayEditorSound(TEXT("/Engine/EditorSounds/Notifications/CompileSuccess_Cue.CompileSuccess_Cue"));
-				}
+				GEditor->PlayEditorSound(TEXT("/Engine/EditorSounds/Notifications/CompileFailed_Cue.CompileFailed_Cue"));
 			}
-
+			else
+			{
+				GEditor->PlayEditorSound(TEXT("/Engine/EditorSounds/Notifications/CompileSuccess_Cue.CompileSuccess_Cue"));
+			}
+			
 			TSharedPtr<SNotificationItem> NotificationItem = NotificationItemPtr.Pin();
 			NotificationItem->SetText(Text);			
 
@@ -486,11 +482,7 @@ public:
 		FEditorDelegates::OnShutdownPostPackagesSaved.Add(FSimpleDelegate::CreateStatic(&FUATHelperModule::HandleUatCancelButtonClicked, UatProcessPtr));
 
 		UatProcess->Launch();
-		const ULevelEditorPlaySettings* EditorPlaySettings = GetDefault<ULevelEditorPlaySettings>();
-		if (EditorPlaySettings->EnablePIEEnterAndExitSounds)
-		{
-			GEditor->PlayEditorSound(TEXT("/Engine/EditorSounds/Notifications/CompileStart_Cue.CompileStart_Cue"));
-		}
+		GEditor->PlayEditorSound(TEXT("/Engine/EditorSounds/Notifications/CompileStart_Cue.CompileStart_Cue"));
 	}
 
 	static void HandleUatHyperlinkNavigate()
@@ -658,11 +650,7 @@ public:
 
 	static void HandleUatLaunchFailed(TWeakPtr<SNotificationItem> NotificationItemPtr, FText PlatformDisplayName, FText TaskName, EventData Event)
 	{
-		const ULevelEditorPlaySettings* EditorPlaySettings = GetDefault<ULevelEditorPlaySettings>();
-		if (EditorPlaySettings->EnablePIEEnterAndExitSounds)
-		{
-			GEditor->PlayEditorSound(TEXT("/Engine/EditorSounds/Notifications/CompileFailed_Cue.CompileFailed_Cue"));
-		}
+		GEditor->PlayEditorSound(TEXT("/Engine/EditorSounds/Notifications/CompileFailed_Cue.CompileFailed_Cue"));
 
 		TGraphTask<FMainFrameActionsNotificationTask>::CreateTask().ConstructAndDispatchWhenReady(
 			NotificationItemPtr,
