Added explicit module dependencies and fixed incorrect includes directive containing incorrect path elements to bypass the missing module dependencies.

diff --git a/Engine/Source/Editor/UATHelper/UATHelper.Build.cs b/Engine/Source/Editor/UATHelper/UATHelper.Build.cs
--- a/Engine/Source/Editor/UATHelper/UATHelper.Build.cs
+++ b/Engine/Source/Editor/UATHelper/UATHelper.Build.cs
@@ -4,22 +4,23 @@ namespace UnrealBuildTool.Rules
 {
 	public class UATHelper : ModuleRules
 	{
-        public UATHelper(ReadOnlyTargetRules Target) : base(Target)
+		public UATHelper(ReadOnlyTargetRules Target) : base(Target)
 		{
 			PrivateDependencyModuleNames.AddRange(
-                new string[] {
-				    "Core",
-				    "CoreUObject",
-				    "Engine",
-                    "InputCore",
-				    "Slate",
+				new string[] {
+					"Core",
+					"CoreUObject",
+					"Engine",
+					"InputCore",
+					"Slate",
 					"SlateCore",
-                    "GameProjectGeneration",
+					"GameProjectGeneration",
 					"EditorFramework",
-                    "UnrealEd",
+					"UnrealEd",
 					"Analytics",
-					"OutputLog", 
-			    }
+					"OutputLog",
+					"MessageLog",
+				}
 			);
 		}
 	}
diff --git a/Engine/Source/Editor/UATHelper/UATHelperModule.cpp b/Engine/Source/Editor/UATHelper/UATHelperModule.cpp
--- a/Engine/Source/Editor/UATHelper/UATHelperModule.cpp
+++ b/Engine/Source/Editor/UATHelper/UATHelperModule.cpp
@@ -28,15 +28,15 @@
 #include "Widgets/Notifications/SNotificationList.h"
 #include "Logging/TokenizedMessage.h"
 #include "Logging/MessageLog.h"
-#include "Developer/MessageLog/Public/IMessageLogListing.h"
-#include "Developer/MessageLog/Public/MessageLogModule.h"
+#include "IMessageLogListing.h"
+#include "MessageLogModule.h"
 #include "Misc/UObjectToken.h"
 
 #include "GameProjectGenerationModule.h"
 #include "AnalyticsEventAttribute.h"
 
 #include "ShaderCompiler.h"
-#include "Developer/OutputLog/Public/OutputLogModule.h"
+#include "OutputLogModule.h"
 
 #define LOCTEXT_NAMESPACE "UATHelper"
 
