Added structured log support for UAT in UE
Fixed UE ignoring warnings after cook during packaging

diff --git a/Engine/Source/Editor/TurnkeySupport/Private/TurnkeySupportModule.cpp b/Engine/Source/Editor/TurnkeySupport/Private/TurnkeySupportModule.cpp
--- a/Engine/Source/Editor/TurnkeySupport/Private/TurnkeySupportModule.cpp
+++ b/Engine/Source/Editor/TurnkeySupport/Private/TurnkeySupportModule.cpp
@@ -393,6 +393,12 @@ public:
 			BuildCookRunParams += TEXT(" -zenstore");
 		}
 
+		// don't report warnings and errors twice
+		BuildCookRunParams += TEXT(" -SkipCookingErrorSummary");
+
+		// use structured logging if available
+		BuildCookRunParams += TEXT(" -JsonStdOut");
+
 		// gather analytics
 		const ITargetPlatform* TargetPlatform = GetTargetPlatformManager()->FindTargetPlatform(PlatformInfo->Name);
 		TargetPlatform->GetPlatformSpecificProjectAnalytics( AnalyticsParamArray );
diff --git a/Engine/Source/Editor/UATHelper/UATHelper.Build.cs b/Engine/Source/Editor/UATHelper/UATHelper.Build.cs
--- a/Engine/Source/Editor/UATHelper/UATHelper.Build.cs
+++ b/Engine/Source/Editor/UATHelper/UATHelper.Build.cs
@@ -19,7 +19,8 @@ namespace UnrealBuildTool.Rules
 					"UnrealEd",
 					"Analytics",
 					"OutputLog",
-			    }
+					"Json"
+				}
 			);
 
 			PrivateIncludePathModuleNames.AddRange(
diff --git a/Engine/Source/Editor/UATHelper/UATHelperModule.cpp b/Engine/Source/Editor/UATHelper/UATHelperModule.cpp
--- a/Engine/Source/Editor/UATHelper/UATHelperModule.cpp
+++ b/Engine/Source/Editor/UATHelper/UATHelperModule.cpp
@@ -20,6 +20,7 @@
 #include "EditorAnalytics.h"
 #include "IUATHelperModule.h"
 #include "AssetRegistry/AssetRegistryModule.h"
+#include "Serialization/JsonSerializer.h"
 
 #include "Interfaces/IMainFrameModule.h"
 #include "Editor/EditorPerProjectUserSettings.h"
@@ -185,8 +186,8 @@ private:
 	* @Param MessageType - The severity of the message, i.e. error, warning etc.
 	**/
 	static void AddMessageToMessageLog(FString MessageString, EMessageSeverity::Type MessageType)
-	{		
-		if (!bSawSummary && (MessageType == EMessageSeverity::Error || MessageType == EMessageSeverity::Warning))
+	{
+		if ((MessageType == EMessageSeverity::Error || MessageType == EMessageSeverity::Warning))
 		{
 			FAssetData AssetData;
 
@@ -276,8 +277,6 @@ private:
 
 	// Whether there are asset errors in the cook, which can be navigated to in the content browser.
 	static bool bHasAssetErrors;
-	// Whether the cook summary has been seen in the log.
-	static bool bSawSummary;
 
 public:
 
@@ -313,31 +312,23 @@ public:
 	{
 		FString LhsUATOutputMsg, ParsedCookIssue;
 
-		// we don't want to report duplicate warnings/errors to the package results log
-		// so, only add messages between the cook start and the summary
-		if (UATOutput.Contains(TEXT("Warning/Error Summary")))
-		{
-			bSawSummary = true;
-		}
-
 		// note: CookResults:Warning: actually outputs some unhandled errors.
-		if ( UATOutput.Split(TEXT("CookResults:Warning: "), &LhsUATOutputMsg, &ParsedCookIssue) )
+		if (UATOutput.Split(TEXT("CookResults:Warning: "), &LhsUATOutputMsg, &ParsedCookIssue) )
 		{
 			SyncMessageWithMessageLog(ParsedCookIssue, EMessageSeverity::Warning);
 		}
-		else if ( UATOutput.Split(TEXT("CookResults:Error: "), &LhsUATOutputMsg, &ParsedCookIssue) )
+		else if (UATOutput.Split(TEXT("CookResults:Error: "), &LhsUATOutputMsg, &ParsedCookIssue) )
 		{
 			SyncMessageWithMessageLog(ParsedCookIssue, EMessageSeverity::Error);
 		}
-		else if (!bSawSummary && UATOutput.Split(TEXT("Warning: "), &LhsUATOutputMsg, &ParsedCookIssue))
+		else if (UATOutput.Split(TEXT("Warning: "), &LhsUATOutputMsg, &ParsedCookIssue))
 		{
 			SyncMessageWithMessageLog(ParsedCookIssue, EMessageSeverity::Warning);
 		}
-		else if (!bSawSummary && UATOutput.Split(TEXT("Error: "), &LhsUATOutputMsg, &ParsedCookIssue))
+		else if (UATOutput.Split(TEXT("Error: "), &LhsUATOutputMsg, &ParsedCookIssue))
 		{
 			SyncMessageWithMessageLog(ParsedCookIssue, EMessageSeverity::Error);
 		}
-
 	}
 
 	/**
@@ -358,12 +349,11 @@ public:
 	/**
 	* Clear the asset error state for the UAT process
 	**/
-	static void ClearAssetErrors() { bHasAssetErrors = false; bSawSummary = false; }
+	static void ClearAssetErrors() { bHasAssetErrors = false; }
 
 };
 
 bool FPackagingErrorHandler::bHasAssetErrors = false;
-bool FPackagingErrorHandler::bSawSummary = false;
 
 DECLARE_CYCLE_STAT(TEXT("Requesting FUATHelperModule::HandleUatProcessCompleted message dialog to present the error message"), STAT_FUATHelperModule_HandleUatProcessCompleted_DialogMessage, STATGROUP_TaskGraphTasks);
 
@@ -629,31 +619,84 @@ public:
 		}
 	}
 
+	static ELogVerbosity::Type GetVerbosityFromStructuredLogString(FStringView Level)
+	{
+		if (Level == TEXT("Information"))
+		{
+			return ELogVerbosity::Log; // could be Display
+		}
+		else if (Level == TEXT("Debug"))
+		{
+			return ELogVerbosity::Verbose; // could be VeryVerbose
+		}
+		else if (Level == TEXT("Warning"))
+		{
+			return ELogVerbosity::Warning;
+		}
+		else if (Level == TEXT("Error"))
+		{
+			return ELogVerbosity::Error;
+		}
+		else if (Level == TEXT("Critical"))
+		{
+			return ELogVerbosity::Fatal;
+		}
+		else
+		{
+			return ELogVerbosity::Log;
+		}
+	}
+
 	static void HandleUatProcessOutput(FString Output, TWeakPtr<SNotificationItem> NotificationItemPtr, FText PlatformDisplayName, FText TaskName)
 	{
-		if ( !Output.IsEmpty() && !Output.Equals("\r") )
+		if (Output.IsEmpty() || Output.Equals("\r"))
 		{
-			bool bDisplayLog = true;
-			if (TaskName.EqualTo(LOCTEXT("PackagingTaskName", "Packaging")))
-			{
-				// Deal with any cook messages that may have been encountered.
-				bDisplayLog = FPackagingErrorHandler::ProcessAndHandleCookMessageOutput(Output);
-			}
-			if (bDisplayLog)
-			{
-				UE_LOG(UATHelper, Log, TEXT("%s (%s): %s"), *TaskName.ToString(), *PlatformDisplayName.ToString(), *Output);
-			}
+			return;
+		}
 
-			if ( TaskName.EqualTo(LOCTEXT("PackagingTaskName", "Packaging")) )
-			{
-				// Deal with any cook errors that may have been encountered.
-				FPackagingErrorHandler::ProcessAndHandleCookErrorOutput(Output);
-			}
-			if (TaskName.EqualTo(LOCTEXT("CookingTaskName", "Cooking")))
+		ELogVerbosity::Type Verbosity = ELogVerbosity::Log;
+		FString Message = Output;
+		bool bStructuredMessage = false;
+
+		// attempt to parse structured logging
+		TSharedRef< TJsonReader<> > JsonReader = FJsonStringReader::Create(Output);
+		TSharedPtr<FJsonObject> JsonData;
+		if (FJsonSerializer::Deserialize(JsonReader, JsonData))
+		{
+			TSharedPtr<FJsonValue> LevelField = JsonData->TryGetField(TEXT("level"));
+			TSharedPtr<FJsonValue> MessageField = JsonData->TryGetField(TEXT("message"));
+
+			FString LevelString;
+			FString MessageString;
+
+			if (LevelField && LevelField->TryGetString(LevelString)
+				&& MessageField && MessageField->TryGetString(Message))
 			{
-				// Deal with any cook errors that may have been encountered.
-				FPackagingErrorHandler::ProcessAndHandleCookErrorOutput(Output);
+				Verbosity = GetVerbosityFromStructuredLogString(LevelString);
+				bStructuredMessage = true;
 			}
+
+			// TODO add support for "format" field
+		}
+
+		bool bDisplayLog = true;
+		if (TaskName.EqualTo(LOCTEXT("PackagingTaskName", "Packaging")))
+		{
+			// Deal with any cook messages that may have been encountered.
+			bDisplayLog = FPackagingErrorHandler::ProcessAndHandleCookMessageOutput(Output);
+		}
+		if (bDisplayLog)
+		{
+			FMsg::Logf(__FILE__, __LINE__,
+				UATHelper.GetCategoryName(),
+				Verbosity,
+				TEXT("%s (%s): %s"), *TaskName.ToString(), *PlatformDisplayName.ToString(), *Message);
+		}
+
+		// deal with non-structured packaging/cooking warnings and errors
+		if (!bStructuredMessage && (TaskName.EqualTo(LOCTEXT("PackagingTaskName", "Packaging")) || TaskName.EqualTo(LOCTEXT("CookingTaskName", "Cooking"))))
+		{
+			FPackagingErrorHandler::ProcessAndHandleCookErrorOutput(Output);
 		}
 	}
 
diff --git a/Engine/Source/Programs/AutomationTool/AutomationUtils/ProjectParams.cs b/Engine/Source/Programs/AutomationTool/AutomationUtils/ProjectParams.cs
--- a/Engine/Source/Programs/AutomationTool/AutomationUtils/ProjectParams.cs
+++ b/Engine/Source/Programs/AutomationTool/AutomationUtils/ProjectParams.cs
@@ -368,6 +368,7 @@ namespace AutomationTool
             this.UnversionedCookedContent = InParams.UnversionedCookedContent;
 			this.OptionalContent = InParams.OptionalContent;
 			this.SkipCookingEditorContent = InParams.SkipCookingEditorContent;
+			this.SkipCookingErrorSummary = InParams.SkipCookingErrorSummary;
 			this.FileServer = InParams.FileServer;
 			this.DedicatedServer = InParams.DedicatedServer;
 			this.Client = InParams.Client;
@@ -529,6 +530,7 @@ namespace AutomationTool
             bool? EncryptPakIndex = null,
 			bool? EncryptEverything = null,
 			bool? SkipCookingEditorContent = null,
+			bool? SkipCookingErrorSummary = null,
             string AdditionalCookerOptions = null,
 			string OriginalReleaseVersion = null,
 			string BasedOnReleaseVersion = null,
@@ -841,6 +843,7 @@ namespace AutomationTool
             this.UnversionedCookedContent = GetOptionalParamValueIfNotSpecified(Command, UnversionedCookedContent, this.UnversionedCookedContent, "UnversionedCookedContent", "VersionCookedContent");
 			this.OptionalContent = GetOptionalParamValueIfNotSpecified(Command, OptionalContent, this.OptionalContent, "editoroptional", "noeditoroptional");
 			this.SkipCookingEditorContent = GetParamValueIfNotSpecified(Command, SkipCookingEditorContent, this.SkipCookingEditorContent, "SkipCookingEditorContent");
+			this.SkipCookingErrorSummary = GetParamValueIfNotSpecified(Command, SkipCookingErrorSummary, this.SkipCookingErrorSummary, "SkipCookingErrorSummary");
 			this.Compressed = GetParamValueIfNotSpecified(Command, Compressed, this.Compressed, "compressed");
 			this.ForceUncompressed = GetParamValueIfNotSpecified(Command, ForceUncompressed, this.ForceUncompressed, "ForceUncompressed");
 			this.UseMergedModules = GetParamValueIfNotSpecified(Command, UseMergedModules, this.UseMergedModules, "UseMergedModules");
@@ -2097,6 +2100,9 @@ namespace AutomationTool
 		[Help("SkipCookingEditorContent", "Skips content under /Engine/Editor when cooking")]
         public bool SkipCookingEditorContent;
 
+		[Help("SkipCookingErrorSummary", "Skips error summary when cooking")]
+		public bool SkipCookingErrorSummary;
+
 		/// <summary>
 		/// Cook: Uses the iterative deploy, command line: -iterativedeploy or -iterate
 		/// </summary>
@@ -3409,6 +3415,7 @@ namespace AutomationTool
 				Logger.LogDebug("UnversionedCookedContent={UnversionedCookedContent}", UnversionedCookedContent);
 				Logger.LogDebug("OptionalContent={OptionalContent}", OptionalContent);
 				Logger.LogDebug("SkipCookingEditorContent={SkipCookingEditorContent}", SkipCookingEditorContent);
+				Logger.LogDebug("SkipCookingErrorSummary={SkipCookingErrorSummary}", SkipCookingErrorSummary);
                 Logger.LogDebug("GeneratePatch={GeneratePatch}", GeneratePatch);
 				Logger.LogDebug("AddPatchLevel={AddPatchLevel}", AddPatchLevel);
 				Logger.LogDebug("DiscVersion={DiscVersion}", DiscVersion);
diff --git a/Engine/Source/Programs/AutomationTool/Scripts/CookCommand.Automation.cs b/Engine/Source/Programs/AutomationTool/Scripts/CookCommand.Automation.cs
--- a/Engine/Source/Programs/AutomationTool/Scripts/CookCommand.Automation.cs
+++ b/Engine/Source/Programs/AutomationTool/Scripts/CookCommand.Automation.cs
@@ -67,6 +67,10 @@ namespace AutomationScripts
 			{
 				CommandletParams += " -skipeditorcontent";
 			}
+			if (Params.SkipCookingErrorSummary)
+			{
+				CommandletParams += " -NoErrorSummary";
+			}
 			if (!String.IsNullOrEmpty(Params.CookOutputDir))
 			{
 				CommandletParams += " -outputdir=" + CommandUtils.MakePathSafeToUseWithCommandLine(Params.CookOutputDir);
diff --git a/Engine/Source/Programs/Shared/EpicGames.Core/Log.cs b/Engine/Source/Programs/Shared/EpicGames.Core/Log.cs
--- a/Engine/Source/Programs/Shared/EpicGames.Core/Log.cs
+++ b/Engine/Source/Programs/Shared/EpicGames.Core/Log.cs
@@ -1080,8 +1080,9 @@ namespace EpicGames.Core
 			IncludeStartingTimestamp = true;
 			EventParser = new LogEventParser(this);
 
+			bool bJsonStdOut = Environment.CommandLine.Contains("JsonStdOut", StringComparison.InvariantCulture);
 			string? envVar = Environment.GetEnvironmentVariable("UE_LOG_JSON_TO_STDOUT");
-			if(envVar != null && Int32.TryParse(envVar, out int value) && value != 0)
+			if(bJsonStdOut || envVar != null && Int32.TryParse(envVar, out int value) && value != 0)
 			{
 				WriteJsonToStdOut = true;
 			}
