Changing UAT critical error to non-fatal error in UE
#rnx

diff --git a/Engine/Source/Editor/UATHelper/UATHelperModule.cpp b/Engine/Source/Editor/UATHelper/UATHelperModule.cpp
--- a/Engine/Source/Editor/UATHelper/UATHelperModule.cpp
+++ b/Engine/Source/Editor/UATHelper/UATHelperModule.cpp
@@ -639,7 +639,7 @@ public:
 		}
 		else if (Level == TEXT("Critical"))
 		{
-			return ELogVerbosity::Fatal;
+			return ELogVerbosity::Error; // don't report external criticals as Fatal because they will then crash our process
 		}
 		else
 		{
